Update Android TargetSDK
Fix templates with AFS configuration

#android

diff --git a/Engine/Config/BaseEngine.ini b/Engine/Config/BaseEngine.ini
--- a/Engine/Config/BaseEngine.ini
+++ b/Engine/Config/BaseEngine.ini
@@ -3272,7 +3272,7 @@ StoreVersionOffsetX86=0
 StoreVersionOffsetX8664=0
 VersionDisplayName=1.0
 MinSDKVersion=26
-TargetSDKVersion=28
+TargetSDKVersion=34
 bEnableLint=false
 bShowLaunchImage=true
 bValidateTextureFormats=true
diff --git a/Templates/TP_AEC_ArchvisBP/Config/DefaultEngine.ini b/Templates/TP_AEC_ArchvisBP/Config/DefaultEngine.ini
--- a/Templates/TP_AEC_ArchvisBP/Config/DefaultEngine.ini
+++ b/Templates/TP_AEC_ArchvisBP/Config/DefaultEngine.ini
@@ -202,18 +202,3 @@ MinSampleRate=8000.000000
 CompressionQualityModifier=1.000000
 AutoStreamingThreshold=0.000000
 SoundCueCookQualityIndex=-1
-
-[/Script/AndroidFileServerEditor.AndroidFileServerRuntimeSettings]
-bEnablePlugin=True
-bAllowNetworkConnection=True
-SecurityToken=4F377AFB4FEB1E99A46E548DDE3402D3
-bIncludeInShipping=False
-bAllowExternalStartInShipping=False
-bCompileAFSProject=False
-bUseCompression=False
-bLogFiles=False
-bReportStats=False
-ConnectionType=USBOnly
-bUseManualIPAddress=False
-ManualIPAddress=
-
diff --git a/Templates/TP_AEC_BlankBP/Config/DefaultEngine.ini b/Templates/TP_AEC_BlankBP/Config/DefaultEngine.ini
--- a/Templates/TP_AEC_BlankBP/Config/DefaultEngine.ini
+++ b/Templates/TP_AEC_BlankBP/Config/DefaultEngine.ini
@@ -151,18 +151,3 @@ SoundCueCookQualityIndex=-1
 +CollisionChannelRedirects=(OldName="Dynamic",NewName="WorldDynamic")
 +CollisionChannelRedirects=(OldName="VehicleMovement",NewName="Vehicle")
 +CollisionChannelRedirects=(OldName="PawnMovement",NewName="Pawn")
-
-[/Script/AndroidFileServerEditor.AndroidFileServerRuntimeSettings]
-bEnablePlugin=True
-bAllowNetworkConnection=True
-SecurityToken=5C8924D94D227DCEFACC8BA745A62272
-bIncludeInShipping=False
-bAllowExternalStartInShipping=False
-bCompileAFSProject=False
-bUseCompression=False
-bLogFiles=False
-bReportStats=False
-ConnectionType=USBOnly
-bUseManualIPAddress=False
-ManualIPAddress=
-
diff --git a/Templates/TP_AEC_CollabBP/Config/DefaultEngine.ini b/Templates/TP_AEC_CollabBP/Config/DefaultEngine.ini
--- a/Templates/TP_AEC_CollabBP/Config/DefaultEngine.ini
+++ b/Templates/TP_AEC_CollabBP/Config/DefaultEngine.ini
@@ -164,7 +164,6 @@ StoreVersionOffsetX8664=0
 ApplicationDisplayName=
 VersionDisplayName=1.0
 MinSDKVersion=26
-TargetSDKVersion=28
 InstallLocation=InternalOnly
 bEnableLint=False
 bPackageDataInsideApk=False
diff --git a/Templates/TP_FirstPerson/Config/DefaultEngine.ini b/Templates/TP_FirstPerson/Config/DefaultEngine.ini
--- a/Templates/TP_FirstPerson/Config/DefaultEngine.ini
+++ b/Templates/TP_FirstPerson/Config/DefaultEngine.ini
@@ -57,20 +57,6 @@ SoundCueCookQualityIndex=-1
 -TargetedRHIs=SF_VULKAN_SM5
 +TargetedRHIs=SF_VULKAN_SM6
 
-[/Script/AndroidFileServerEditor.AndroidFileServerRuntimeSettings]
-bEnablePlugin=True
-bAllowNetworkConnection=True
-SecurityToken=7B0AA718435214EB1E253DACC60399C4
-bIncludeInShipping=False
-bAllowExternalStartInShipping=False
-bCompileAFSProject=False
-bUseCompression=False
-bLogFiles=False
-bReportStats=False
-ConnectionType=USBOnly
-bUseManualIPAddress=False
-ManualIPAddress=
-
 [/Script/AIModule.AISystem]
 bForgetStaleActors=True
 
diff --git a/Templates/TP_FirstPersonBP/Config/DefaultEngine.ini b/Templates/TP_FirstPersonBP/Config/DefaultEngine.ini
--- a/Templates/TP_FirstPersonBP/Config/DefaultEngine.ini
+++ b/Templates/TP_FirstPersonBP/Config/DefaultEngine.ini
@@ -58,20 +58,6 @@ AppliedDefaultGraphicsPerformance=Maximum
 +ActiveGameNameRedirects=(OldGameName="/Script/TP_FirstPersonBP",NewGameName="/Script/MyProject")
 NearClipPlane=5.000000
 
-[/Script/AndroidFileServerEditor.AndroidFileServerRuntimeSettings]
-bEnablePlugin=True
-bAllowNetworkConnection=True
-SecurityToken=85DFF0BB498AE9EFE914DE8DB82D122D
-bIncludeInShipping=False
-bAllowExternalStartInShipping=False
-bCompileAFSProject=False
-bUseCompression=False
-bLogFiles=False
-bReportStats=False
-ConnectionType=USBOnly
-bUseManualIPAddress=False
-ManualIPAddress=
-
 [/Script/AIModule.AISystem]
 bForgetStaleActors=True
 
diff --git a/Templates/TP_ME_VProdBP/Config/DefaultEngine.ini b/Templates/TP_ME_VProdBP/Config/DefaultEngine.ini
--- a/Templates/TP_ME_VProdBP/Config/DefaultEngine.ini
+++ b/Templates/TP_ME_VProdBP/Config/DefaultEngine.ini
@@ -109,18 +109,3 @@ DefaultGraphicsRHI=DefaultGraphicsRHI_DX12
 [/Script/LinuxTargetPlatform.LinuxTargetSettings]
 -TargetedRHIs=SF_VULKAN_SM5
 +TargetedRHIs=SF_VULKAN_SM6
-
-[/Script/AndroidFileServerEditor.AndroidFileServerRuntimeSettings]
-bEnablePlugin=True
-bAllowNetworkConnection=True
-SecurityToken=FE1C85674E26FE981BF51B824609C20E
-bIncludeInShipping=False
-bAllowExternalStartInShipping=False
-bCompileAFSProject=False
-bUseCompression=False
-bLogFiles=False
-bReportStats=False
-ConnectionType=USBOnly
-bUseManualIPAddress=False
-ManualIPAddress=
-
diff --git a/Templates/TP_MotionDesignBP/Config/DefaultEngine.ini b/Templates/TP_MotionDesignBP/Config/DefaultEngine.ini
--- a/Templates/TP_MotionDesignBP/Config/DefaultEngine.ini
+++ b/Templates/TP_MotionDesignBP/Config/DefaultEngine.ini
@@ -128,18 +128,3 @@ DefaultGraphicsRHI=DefaultGraphicsRHI_DX12
 [/Script/Engine.Engine]
 +ActiveGameNameRedirects=(OldGameName="TP_ME_BlankBP",NewGameName="/Script/TP_MotionDesign")
 +ActiveGameNameRedirects=(OldGameName="/Script/TP_ME_BlankBP",NewGameName="/Script/TP_MotionDesign")
-
-[/Script/AndroidFileServerEditor.AndroidFileServerRuntimeSettings]
-bEnablePlugin=True
-bAllowNetworkConnection=True
-SecurityToken=1D0AC1BA4E92374633D885A9A0A56B49
-bIncludeInShipping=False
-bAllowExternalStartInShipping=False
-bCompileAFSProject=False
-bUseCompression=False
-bLogFiles=False
-bReportStats=False
-ConnectionType=USBOnly
-bUseManualIPAddress=False
-ManualIPAddress=
-
diff --git a/Templates/TP_ThirdPerson/Config/DefaultEngine.ini b/Templates/TP_ThirdPerson/Config/DefaultEngine.ini
--- a/Templates/TP_ThirdPerson/Config/DefaultEngine.ini
+++ b/Templates/TP_ThirdPerson/Config/DefaultEngine.ini
@@ -101,18 +101,3 @@ SoundCueCookQualityIndex=-1
 +CollisionChannelRedirects=(OldName="Dynamic",NewName="WorldDynamic")
 +CollisionChannelRedirects=(OldName="VehicleMovement",NewName="Vehicle")
 +CollisionChannelRedirects=(OldName="PawnMovement",NewName="Pawn")
-
-[/Script/AndroidFileServerEditor.AndroidFileServerRuntimeSettings]
-bEnablePlugin=True
-bAllowNetworkConnection=True
-SecurityToken=F022329C4DC0428B5B294D852830B9FD
-bIncludeInShipping=False
-bAllowExternalStartInShipping=False
-bCompileAFSProject=False
-bUseCompression=False
-bLogFiles=False
-bReportStats=False
-ConnectionType=USBOnly
-bUseManualIPAddress=False
-ManualIPAddress=
-
diff --git a/Templates/TP_ThirdPersonBP/Config/DefaultEngine.ini b/Templates/TP_ThirdPersonBP/Config/DefaultEngine.ini
--- a/Templates/TP_ThirdPersonBP/Config/DefaultEngine.ini
+++ b/Templates/TP_ThirdPersonBP/Config/DefaultEngine.ini
@@ -52,18 +52,3 @@ SoundCueCookQualityIndex=-1
 [/Script/LinuxTargetPlatform.LinuxTargetSettings]
 -TargetedRHIs=SF_VULKAN_SM5
 +TargetedRHIs=SF_VULKAN_SM6
-
-[/Script/AndroidFileServerEditor.AndroidFileServerRuntimeSettings]
-bEnablePlugin=True
-bAllowNetworkConnection=True
-SecurityToken=9A8C9D134D11008E5B456E90FC2B7A65
-bIncludeInShipping=False
-bAllowExternalStartInShipping=False
-bCompileAFSProject=False
-bUseCompression=False
-bLogFiles=False
-bReportStats=False
-ConnectionType=USBOnly
-bUseManualIPAddress=False
-ManualIPAddress=
-
diff --git a/Templates/TP_TopDown/Config/DefaultEngine.ini b/Templates/TP_TopDown/Config/DefaultEngine.ini
--- a/Templates/TP_TopDown/Config/DefaultEngine.ini
+++ b/Templates/TP_TopDown/Config/DefaultEngine.ini
@@ -63,20 +63,6 @@ SoundCueCookQualityIndex=-1
 ;+StructRedirects=(OldName="MyStruct",NewName="MyNewStruct")
 ;+StructRedirects=(OldName="MyStruct",NewName="MyNewStruct")
 
-[/Script/AndroidFileServerEditor.AndroidFileServerRuntimeSettings]
-bEnablePlugin=True
-bAllowNetworkConnection=True
-SecurityToken=0FB14AED42FEC3066D998A8660BACCA3
-bIncludeInShipping=False
-bAllowExternalStartInShipping=False
-bCompileAFSProject=False
-bUseCompression=False
-bLogFiles=False
-bReportStats=False
-ConnectionType=USBOnly
-bUseManualIPAddress=False
-ManualIPAddress=
-
 [/Script/NavigationSystem.RecastNavMesh]
 RuntimeGeneration=Dynamic
 
diff --git a/Templates/TP_TopDownBP/Config/DefaultEngine.ini b/Templates/TP_TopDownBP/Config/DefaultEngine.ini
--- a/Templates/TP_TopDownBP/Config/DefaultEngine.ini
+++ b/Templates/TP_TopDownBP/Config/DefaultEngine.ini
@@ -90,18 +90,3 @@ SoundCueCookQualityIndex=-1
 [/Script/LinuxTargetPlatform.LinuxTargetSettings]
 -TargetedRHIs=SF_VULKAN_SM5
 +TargetedRHIs=SF_VULKAN_SM6
-
-[/Script/AndroidFileServerEditor.AndroidFileServerRuntimeSettings]
-bEnablePlugin=True
-bAllowNetworkConnection=True
-SecurityToken=7F2D0E674D61A6F3260B0082E348BC25
-bIncludeInShipping=False
-bAllowExternalStartInShipping=False
-bCompileAFSProject=False
-bUseCompression=False
-bLogFiles=False
-bReportStats=False
-ConnectionType=USBOnly
-bUseManualIPAddress=False
-ManualIPAddress=
-
diff --git a/Templates/TP_VehicleAdv/Config/DefaultEngine.ini b/Templates/TP_VehicleAdv/Config/DefaultEngine.ini
--- a/Templates/TP_VehicleAdv/Config/DefaultEngine.ini
+++ b/Templates/TP_VehicleAdv/Config/DefaultEngine.ini
@@ -53,18 +53,3 @@ SoundCueCookQualityIndex=-1
 [/Script/LinuxTargetPlatform.LinuxTargetSettings]
 -TargetedRHIs=SF_VULKAN_SM5
 +TargetedRHIs=SF_VULKAN_SM6
-
-[/Script/AndroidFileServerEditor.AndroidFileServerRuntimeSettings]
-bEnablePlugin=True
-bAllowNetworkConnection=True
-SecurityToken=7584171B4FAB7742F2593A9D8AE0B9F0
-bIncludeInShipping=False
-bAllowExternalStartInShipping=False
-bCompileAFSProject=False
-bUseCompression=False
-bLogFiles=False
-bReportStats=False
-ConnectionType=USBOnly
-bUseManualIPAddress=False
-ManualIPAddress=
-
diff --git a/Templates/TP_VehicleAdvBP/Config/DefaultEngine.ini b/Templates/TP_VehicleAdvBP/Config/DefaultEngine.ini
--- a/Templates/TP_VehicleAdvBP/Config/DefaultEngine.ini
+++ b/Templates/TP_VehicleAdvBP/Config/DefaultEngine.ini
@@ -82,18 +82,3 @@ AppliedDefaultGraphicsPerformance=Maximum
 [/Script/Engine.Engine]
 +ActiveGameNameRedirects=(OldGameName="TP_VehicleAdvBP",NewGameName="/Script/VehicleTemplateNew")
 +ActiveGameNameRedirects=(OldGameName="/Script/TP_VehicleAdvBP",NewGameName="/Script/VehicleTemplateNew")
-
-[/Script/AndroidFileServerEditor.AndroidFileServerRuntimeSettings]
-bEnablePlugin=True
-bAllowNetworkConnection=True
-SecurityToken=0AAA16104B3CEB6E0395D7B28752BFBD
-bIncludeInShipping=False
-bAllowExternalStartInShipping=False
-bCompileAFSProject=False
-bUseCompression=False
-bLogFiles=False
-bReportStats=False
-ConnectionType=USBOnly
-bUseManualIPAddress=False
-ManualIPAddress=
-
